Live Link device framework: Add Unreal (Take Recorder) device.

The new ULiveLinkUnrealDevice is the first device in the new LiveLinkStandardDevices plugin, enabled by default in the Live Link Hub target.

Currently primarily serves as Take Recorder control via Message Bus, supporting:
- Slate name / take number synchronization
- Forwarding LLH record start / stop events to Take Recorder
- Optionally responds to Take Recorder start / stop events by starting / stopping Live Link Hub recording (controlled by optional "record start authority" and "record stop authority" checkboxes in device settings, enabled by default).

Other changes:
- Implemented device table health indicator UI elements.
- Adds EDeviceHealth::Good (green circle check icon from UI design).
- "Hardware ID" column now mirrors cell label text to tooltip, so it can be displayed on hover (useful if it overflows).
- Don't register Live Link "Devices" or "Device Details" tab spawners outside of Live Link Hub. Enables users to load the LiveLinkDevice plugin in Unreal Editor without the corresponding tab summoning entries appearing in UE's "Window" menu.

diff --git a/Engine/Plugins/Animation/LiveLinkDevice/Source/LiveLinkDevice/Private/LiveLinkDeviceCapability_Connection.cpp b/Engine/Plugins/Animation/LiveLinkDevice/Source/LiveLinkDevice/Private/LiveLinkDeviceCapability_Connection.cpp
--- a/Engine/Plugins/Animation/LiveLinkDevice/Source/LiveLinkDevice/Private/LiveLinkDeviceCapability_Connection.cpp
+++ b/Engine/Plugins/Animation/LiveLinkDevice/Source/LiveLinkDevice/Private/LiveLinkDeviceCapability_Connection.cpp
@@ -38,6 +38,7 @@ public:
 			.IsSelected(InArgs._IsDeviceRowSelected)
 			.OnTextCommitted(InArgs._OnTextCommitted)
 			.Text(InArgs._Text)
+			.ToolTipText(InArgs._Text)
 		];
 	}
 
diff --git a/Engine/Plugins/Animation/LiveLinkDevice/Source/LiveLinkDevice/Private/LiveLinkDeviceModule.cpp b/Engine/Plugins/Animation/LiveLinkDevice/Source/LiveLinkDevice/Private/LiveLinkDeviceModule.cpp
--- a/Engine/Plugins/Animation/LiveLinkDevice/Source/LiveLinkDevice/Private/LiveLinkDeviceModule.cpp
+++ b/Engine/Plugins/Animation/LiveLinkDevice/Source/LiveLinkDevice/Private/LiveLinkDeviceModule.cpp
@@ -9,6 +9,7 @@
 #include "LiveLinkDeviceStyle.h"
 #include "LiveLinkDeviceSubsystem.h"
 #include "Logging/StructuredLog.h"
+#include "Misc/ConfigCacheIni.h"
 #include "PropertyEditorModule.h"
 #include "SPositiveActionButton.h"
 #include "Widgets/Docking/SDockTab.h"
@@ -35,22 +36,34 @@ void FLiveLinkDeviceModule::StartupModule()
 {
 	FLiveLinkDeviceStyle::Initialize();
 
-	FGlobalTabmanager::Get()
-		->RegisterNomadTabSpawner(DevicesTabName,
-			FOnSpawnTab::CreateRaw(this, &FLiveLinkDeviceModule::OnSpawnDevicesTab))
-		.SetIcon(FSlateIcon("LiveLinkDeviceStyle", "LiveLinkHub.Devices.Icon"))
-		.SetDisplayName(LOCTEXT("DevicesTabDisplayName", "Devices"));
-
-	FGlobalTabmanager::Get()
-		->RegisterNomadTabSpawner(DeviceDetailsTabName,
-			FOnSpawnTab::CreateRaw(this, &FLiveLinkDeviceModule::OnSpawnDeviceDetailsTab))
-		.SetIcon(FSlateIcon(FAppStyle::GetAppStyleSetName(), "LevelEditor.Tabs.Details"))
-		.SetDisplayName(LOCTEXT("DeviceDetailsTabDisplayName", "Device Details"));
+	bIsLiveLinkHubApp = GConfig->GetBoolOrDefault(TEXT("LiveLink"), TEXT("bCreateLiveLinkHubInstance"), false, GEngineIni);
+
+	if (bIsLiveLinkHubApp)
+	{
+		TSharedRef<FGlobalTabmanager> GlobalTabManager = FGlobalTabmanager::Get();
+	
+		GlobalTabManager->RegisterNomadTabSpawner(DevicesTabName,
+				FOnSpawnTab::CreateRaw(this, &FLiveLinkDeviceModule::OnSpawnDevicesTab))
+			.SetIcon(FSlateIcon("LiveLinkDeviceStyle", "LiveLinkHub.Devices.Icon"))
+			.SetDisplayName(LOCTEXT("DevicesTabDisplayName", "Devices"));
+
+		GlobalTabManager->RegisterNomadTabSpawner(DeviceDetailsTabName,
+				FOnSpawnTab::CreateRaw(this, &FLiveLinkDeviceModule::OnSpawnDeviceDetailsTab))
+			.SetIcon(FSlateIcon(FAppStyle::GetAppStyleSetName(), "LevelEditor.Tabs.Details"))
+			.SetDisplayName(LOCTEXT("DeviceDetailsTabDisplayName", "Device Details"));
+	}
 }
 
 
 void FLiveLinkDeviceModule::ShutdownModule()
 {
+	if (bIsLiveLinkHubApp)		
+	{
+		TSharedRef<FGlobalTabmanager> GlobalTabManager = FGlobalTabmanager::Get();
+		GlobalTabManager->UnregisterNomadTabSpawner(DevicesTabName);
+		GlobalTabManager->UnregisterNomadTabSpawner(DeviceDetailsTabName);
+	}
+
 	FLiveLinkDeviceStyle::Shutdown();
 }
 
diff --git a/Engine/Plugins/Animation/LiveLinkDevice/Source/LiveLinkDevice/Private/LiveLinkDeviceModule.h b/Engine/Plugins/Animation/LiveLinkDevice/Source/LiveLinkDevice/Private/LiveLinkDeviceModule.h
--- a/Engine/Plugins/Animation/LiveLinkDevice/Source/LiveLinkDevice/Private/LiveLinkDeviceModule.h
+++ b/Engine/Plugins/Animation/LiveLinkDevice/Source/LiveLinkDevice/Private/LiveLinkDeviceModule.h
@@ -36,6 +36,7 @@ private:
 private:
 	void DeviceSelectionChanged(ULiveLinkDevice* InSelectedDevice);
 
+	bool bIsLiveLinkHubApp = false;
 	TSharedPtr<SLiveLinkDeviceTable> DeviceTable;
 	TSharedPtr<IDetailsView> DetailsView;
 	TWeakObjectPtr<ULiveLinkDevice> WeakSelectedDevice;
diff --git a/Engine/Plugins/Animation/LiveLinkDevice/Source/LiveLinkDevice/Private/LiveLinkDeviceStyle.cpp b/Engine/Plugins/Animation/LiveLinkDevice/Source/LiveLinkDevice/Private/LiveLinkDeviceStyle.cpp
--- a/Engine/Plugins/Animation/LiveLinkDevice/Source/LiveLinkDevice/Private/LiveLinkDeviceStyle.cpp
+++ b/Engine/Plugins/Animation/LiveLinkDevice/Source/LiveLinkDevice/Private/LiveLinkDeviceStyle.cpp
@@ -4,6 +4,7 @@
 #include "Interfaces/IPluginManager.h"
 #include "Styling/SlateStyle.h"
 #include "Styling/SlateStyleRegistry.h"
+#include "Styling/StyleColors.h"
 
 
 TSharedPtr<FSlateStyleSet> FLiveLinkDeviceStyle::StyleSet;
@@ -16,19 +17,27 @@ void FLiveLinkDeviceStyle::Initialize()
 		return;
 	}
 
-	StyleSet = MakeShared<FSlateStyleSet>(GetStyleSetName());
+	static const FVector2D Icon16x16(16.0f, 16.0f);
+	static const FVector2D Icon20x20(20.0f, 20.0f);
 
 	const FString PluginContentRoot =
 		IPluginManager::Get().FindPlugin("LiveLinkDevice")->GetBaseDir() / TEXT("Resources");
-	StyleSet->SetContentRoot(PluginContentRoot);
 
-	static const FVector2D Icon16x16(16.0f, 16.0f);
-	static const FVector2D Icon20x20(20.0f, 20.0f);
+	StyleSet = MakeShared<FSlateStyleSet>(GetStyleSetName());
+
+	StyleSet->SetContentRoot(FPaths::EngineContentDir() / TEXT("Editor/Slate"));
+	StyleSet->SetCoreContentRoot(FPaths::EngineContentDir() / TEXT("Slate"));
+
 	StyleSet->Set("Record", new FSlateVectorImageBrush(PluginContentRoot / "PlayControlsRecord.svg", Icon16x16));
 	StyleSet->Set("Record.Monochrome", new FSlateVectorImageBrush(PluginContentRoot / "PlayControlsRecord_Monochrome.svg", Icon16x16));
 
 	StyleSet->Set("LiveLinkHub.Devices.Icon", new FSlateVectorImageBrush(PluginContentRoot / "Devices.svg", Icon16x16));
 
+	StyleSet->Set("LiveLinkHub.Devices.Status.Good", new FSlateVectorImageBrush(StyleSet->RootToCoreContentDir("Starship/Common/check-circle-solid.svg"), Icon16x16, FStyleColors::Success));
+	StyleSet->Set("LiveLinkHub.Devices.Status.Info", new FSlateVectorImageBrush(StyleSet->RootToCoreContentDir("Starship/Common/info-circle-solid.svg"), Icon16x16));
+	StyleSet->Set("LiveLinkHub.Devices.Status.Warning", new FSlateVectorImageBrush(StyleSet->RootToContentDir("Starship/Common/AlertTriangleSolid.svg"), Icon16x16, FStyleColors::Warning));
+	StyleSet->Set("LiveLinkHub.Devices.Status.Error", new FSlateVectorImageBrush(StyleSet->RootToContentDir("Starship/Common/AlertTriangleSolid.svg"), Icon16x16, FStyleColors::Error));
+
 	FSlateStyleRegistry::RegisterSlateStyle(*StyleSet);
 }
 
diff --git a/Engine/Plugins/Animation/LiveLinkDevice/Source/LiveLinkDevice/Private/Widgets/DeviceTable.cpp b/Engine/Plugins/Animation/LiveLinkDevice/Source/LiveLinkDevice/Private/Widgets/DeviceTable.cpp
--- a/Engine/Plugins/Animation/LiveLinkDevice/Source/LiveLinkDevice/Private/Widgets/DeviceTable.cpp
+++ b/Engine/Plugins/Animation/LiveLinkDevice/Source/LiveLinkDevice/Private/Widgets/DeviceTable.cpp
@@ -6,6 +6,7 @@
 #include "LiveLinkDevice.h"
 #include "LiveLinkDeviceCapability.h"
 #include "LiveLinkDeviceModule.h"
+#include "LiveLinkDeviceStyle.h"
 #include "LiveLinkDeviceSubsystem.h"
 #include "Widgets/Input/SButton.h"
 
@@ -174,16 +175,47 @@ void SLiveLinkDeviceRowWidget::Construct(
 
 TSharedRef<SWidget> SLiveLinkDeviceRowWidget::GenerateWidgetForColumn(const FName& InColumnId)
 {
-	ULiveLinkDevice* Device = GetDevice();
-
-	if (!ensure(Device))
+	if (!ensure(GetDevice()))
 	{
 		return SNullWidget::NullWidget;
 	}
 
+	TWeakObjectPtr<ULiveLinkDevice> WeakDevice = GetDevice();
+
 	if (InColumnId == SLiveLinkDeviceTable::ColumnName_StatusIcon)
 	{
-		return SNullWidget::NullWidget;
+		return SNew(SImage)
+			.Image_Lambda(
+		        [WeakDevice]
+				() -> const FSlateBrush*
+				{
+			        if (ULiveLinkDevice* Device = WeakDevice.Get())
+					{
+						switch (Device->GetDeviceHealth())
+						{
+							case EDeviceHealth::Good: return FLiveLinkDeviceStyle::Get()->GetBrush("LiveLinkHub.Devices.Status.Good");
+							case EDeviceHealth::Info: return FLiveLinkDeviceStyle::Get()->GetBrush("LiveLinkHub.Devices.Status.Info");
+							case EDeviceHealth::Warning: return FLiveLinkDeviceStyle::Get()->GetBrush("LiveLinkHub.Devices.Status.Warning");
+							case EDeviceHealth::Error: return FLiveLinkDeviceStyle::Get()->GetBrush("LiveLinkHub.Devices.Status.Error");
+							default: return nullptr;
+						}
+					}
+
+				    return nullptr;
+				}
+			)
+			.ToolTipText_Lambda(
+				[WeakDevice]
+				() -> FText
+				{
+			        if (ULiveLinkDevice* Device = WeakDevice.Get())
+			        {
+						return Device->GetHealthText();
+			        }
+
+				    return FText::GetEmpty();
+				}
+			);
 	}
 
 	if (InColumnId == SLiveLinkDeviceTable::ColumnName_DisplayName)
@@ -194,8 +226,28 @@ TSharedRef<SWidget> SLiveLinkDeviceRowWidget::GenerateWidgetForColumn(const FNam
 			.BorderImage(FAppStyle::GetBrush("NoBorder"))
 			[
 				SNew(STextBlock)
-				.Text_Lambda([Device]() { return Device->GetDisplayName(); })
-				.ToolTipText_Lambda([Device]() { return Device->GetDisplayName(); })
+				.Text_Lambda(
+					[WeakDevice]
+					()
+					{
+						if (ULiveLinkDevice* Device = WeakDevice.Get())
+						{
+							return Device->GetDisplayName();
+						}
+
+						return FText::GetEmpty();
+					})
+				.ToolTipText_Lambda(
+					[WeakDevice]
+					()
+					{
+						if (ULiveLinkDevice* Device = WeakDevice.Get())
+						{
+							return Device->GetDisplayName();
+						}
+
+						return FText::GetEmpty();
+					})
 			];
 	}
 
@@ -210,10 +262,13 @@ TSharedRef<SWidget> SLiveLinkDeviceRowWidget::GenerateWidgetForColumn(const FNam
 			.ForegroundColor(FSlateColor::UseForeground())
 			.IsFocusable(false)
 			.OnClicked_Lambda(
-				[Device]
+				[WeakDevice]
 				{
-					ULiveLinkDeviceSubsystem* Subsystem = GEngine->GetEngineSubsystem<ULiveLinkDeviceSubsystem>();
-					Subsystem->RemoveDevice(Device);
+			        if (ULiveLinkDevice* Device = WeakDevice.Get())
+			        {
+						ULiveLinkDeviceSubsystem* Subsystem = GEngine->GetEngineSubsystem<ULiveLinkDeviceSubsystem>();
+						Subsystem->RemoveDevice(Device);
+			        }
 					return FReply::Handled();
 				})
 			[
@@ -232,7 +287,7 @@ TSharedRef<SWidget> SLiveLinkDeviceRowWidget::GenerateWidgetForColumn(const FNam
 	}
 
 	const TSubclassOf<ULiveLinkDeviceCapability>& CapabilityClass = ColumnIdToCapability.FindChecked(InColumnId);
-	if (!Device->GetClass()->ImplementsInterface(CapabilityClass))
+	if (!GetDevice()->GetClass()->ImplementsInterface(CapabilityClass))
 	{
 		// Device does not implement the capability responsible for populating this column.
 		return SNullWidget::NullWidget;
@@ -242,7 +297,7 @@ TSharedRef<SWidget> SLiveLinkDeviceRowWidget::GenerateWidgetForColumn(const FNam
 		.IsRowSelected = TDelegate<bool()>::CreateSPLambda(this, [this]() { return IsSelected(); })
 	};
 
-	return Device->GenerateWidgetForColumn(InColumnId, Args);
+	return GetDevice()->GenerateWidgetForColumn(InColumnId, Args);
 }
 
 
diff --git a/Engine/Plugins/Animation/LiveLinkDevice/Source/LiveLinkDevice/Public/LiveLinkDevice.h b/Engine/Plugins/Animation/LiveLinkDevice/Source/LiveLinkDevice/Public/LiveLinkDevice.h
--- a/Engine/Plugins/Animation/LiveLinkDevice/Source/LiveLinkDevice/Public/LiveLinkDevice.h
+++ b/Engine/Plugins/Animation/LiveLinkDevice/Source/LiveLinkDevice/Public/LiveLinkDevice.h
@@ -16,6 +16,7 @@ UENUM(BlueprintType)
 enum class EDeviceHealth : uint8
 {
 	Nominal = 0,
+	Good,
 	Info,
 	Warning,
 	Error,
diff --git a/Engine/Plugins/Animation/LiveLinkDevice/Source/LiveLinkDevice/Public/LiveLinkDeviceCapability_Connection.h b/Engine/Plugins/Animation/LiveLinkDevice/Source/LiveLinkDevice/Public/LiveLinkDeviceCapability_Connection.h
--- a/Engine/Plugins/Animation/LiveLinkDevice/Source/LiveLinkDevice/Public/LiveLinkDeviceCapability_Connection.h
+++ b/Engine/Plugins/Animation/LiveLinkDevice/Source/LiveLinkDevice/Public/LiveLinkDeviceCapability_Connection.h
@@ -120,7 +120,7 @@ public:
 	UConnectionDelegate* GetConnectionDelegate_Implementation();
 
 protected:
-	void SetConnectionStatus(ELiveLinkDeviceConnectionStatus InStatus);
+	virtual void SetConnectionStatus(ELiveLinkDeviceConnectionStatus InStatus);
 
 private:
 	TStrongObjectPtr<UConnectionDelegate> ConnectionDelegate;
diff --git a/Engine/Plugins/Animation/LiveLinkStandardDevices/LiveLinkStandardDevices.uplugin b/Engine/Plugins/Animation/LiveLinkStandardDevices/LiveLinkStandardDevices.uplugin
new file mode 100644
--- /dev/null
+++ b/Engine/Plugins/Animation/LiveLinkStandardDevices/LiveLinkStandardDevices.uplugin
@@ -0,0 +1,31 @@
+{
+	"FileVersion": 3,
+	"Version": 1,
+	"VersionName": "1.0",
+	"FriendlyName": "Live Link Hub Standard Devices",
+	"Description": "Built-in support included for a selection of Live Link device types",
+	"Category": "Animation",
+	"CreatedBy": "Epic Games, Inc.",
+	"CreatedByURL": "https://epicgames.com",
+	"DocsURL": "",
+	"MarketplaceURL": "",
+	"SupportURL": "",
+	"EnabledByDefault": false,
+	"CanContainContent": true,
+	"IsBetaVersion": false,
+	"IsExperimentalVersion": true,
+	"Installed": false,
+	"Plugins": [
+		{
+			"Name": "LiveLinkDevice",
+			"Enabled": true
+		}
+	],
+	"Modules": [
+		{
+			"Name": "LiveLinkStandardDevices",
+			"Type": "Editor",
+			"LoadingPhase": "Default"
+		}
+	]
+}
diff --git a/Engine/Plugins/Animation/LiveLinkStandardDevices/Source/LiveLinkStandardDevices/LiveLinkStandardDevices.Build.cs b/Engine/Plugins/Animation/LiveLinkStandardDevices/Source/LiveLinkStandardDevices/LiveLinkStandardDevices.Build.cs
new file mode 100644
--- /dev/null
+++ b/Engine/Plugins/Animation/LiveLinkStandardDevices/Source/LiveLinkStandardDevices/LiveLinkStandardDevices.Build.cs
@@ -0,0 +1,31 @@
+// Copyright Epic Games, Inc. All Rights Reserved.
+
+using UnrealBuildTool;
+
+public class LiveLinkStandardDevices : ModuleRules
+{
+	public LiveLinkStandardDevices(ReadOnlyTargetRules Target) : base(Target)
+	{
+		PrivateDependencyModuleNames.AddRange(
+			new string[] {
+				"Core",
+				"CoreUObject",
+				"Engine",
+				"InputCore",
+				"LiveLink",
+				"LiveLinkDevice",
+				"LiveLinkHub",
+				"LiveLinkHubMessaging",
+				"LiveLinkSequencer",
+				"Messaging",
+				"Networking",
+				"Projects",
+				"Slate",
+				"SlateCore",
+				"TakeRecorder",
+			}
+		);
+
+		CppCompileWarningSettings.UnsafeTypeCastWarningLevel = WarningLevel.Error;
+	}
+}
diff --git a/Engine/Plugins/Animation/LiveLinkStandardDevices/Source/LiveLinkStandardDevices/Private/Devices/LiveLinkUnrealDevice.cpp b/Engine/Plugins/Animation/LiveLinkStandardDevices/Source/LiveLinkStandardDevices/Private/Devices/LiveLinkUnrealDevice.cpp
new file mode 100644
--- /dev/null
+++ b/Engine/Plugins/Animation/LiveLinkStandardDevices/Source/LiveLinkStandardDevices/Private/Devices/LiveLinkUnrealDevice.cpp
@@ -0,0 +1,380 @@
+// Copyright Epic Games, Inc. All Rights Reserved.
+
+#include "LiveLinkUnrealDevice.h"
+#include "Engine/Engine.h"
+#include "ILiveLinkRecordingSession.h"
+#include "INetworkMessagingExtension.h"
+#include "LiveLinkDeviceSubsystem.h"
+#include "LiveLinkUnrealDeviceMessages.h"
+#include "Logging/StructuredLog.h"
+#include "MessageEndpoint.h"
+#include "MessageEndpointBuilder.h"
+
+
+DEFINE_LOG_CATEGORY(LogLiveLinkUnrealDevice);
+
+
+#define LOCTEXT_NAMESPACE "LiveLinkUnrealDevice"
+
+
+FText ULiveLinkUnrealDevice::GetDisplayName() const
+{
+	return FText::FromString(GetSettings()->DisplayName);
+}
+
+
+EDeviceHealth ULiveLinkUnrealDevice::GetDeviceHealth() const
+{
+	return GetHealth().Key;
+}
+
+
+FText ULiveLinkUnrealDevice::GetHealthText() const
+{
+	return GetHealth().Value;
+}
+
+
+TPair<EDeviceHealth, FText> ULiveLinkUnrealDevice::GetHealth() const
+{
+	const ULiveLinkUnrealDeviceSettings* DeviceSettings = GetDeviceSettings<ULiveLinkUnrealDeviceSettings>();
+	if (!DeviceSettings->ClientId.IsValid())
+	{
+		return { EDeviceHealth::Info, LOCTEXT("DeviceHealthText_ClientUnset", "Client unset") };
+	}
+	else if (!MessageEndpoint || !DeviceAddress.IsValid())
+	{
+		return { EDeviceHealth::Warning, LOCTEXT("DeviceHealthText_NotConnected", "Not connected") };
+	}
+
+	return { EDeviceHealth::Good, FText::GetEmpty() };
+}
+
+
+void ULiveLinkUnrealDevice::OnDeviceAdded()
+{
+	ILiveLinkRecordingSession& Session = ILiveLinkRecordingSession::Get();
+	Session.OnSlateNameChanged().AddUObject(this, &ULiveLinkUnrealDevice::HandleSlateNameChanged);
+	Session.OnTakeNumberChanged().AddUObject(this, &ULiveLinkUnrealDevice::HandleTakeNumberChanged);
+
+	Connect();
+}
+
+
+void ULiveLinkUnrealDevice::OnDeviceRemoved()
+{
+	Disconnect();
+}
+
+
+void ULiveLinkUnrealDevice::OnSettingChanged(const FPropertyChangedEvent& InPropertyChangedEvent)
+{
+	Super::OnSettingChanged(InPropertyChangedEvent);
+
+	const ULiveLinkUnrealDeviceSettings* DeviceSettings = GetDeviceSettings<ULiveLinkUnrealDeviceSettings>();
+
+	const FName ChangedPropertyName = InPropertyChangedEvent.GetPropertyName();
+
+	static const FName ClientIdName = GET_MEMBER_NAME_CHECKED(ULiveLinkUnrealDeviceSettings, ClientId);
+	if (ChangedPropertyName == ClientIdName)
+	{
+		Disconnect();
+		Connect();
+	}
+}
+
+
+ELiveLinkDeviceConnectionStatus ULiveLinkUnrealDevice::GetConnectionStatus_Implementation() const
+{
+	return ConnectionStatus;
+}
+
+
+FString ULiveLinkUnrealDevice::GetHardwareId_Implementation() const
+{
+	const ULiveLinkUnrealDeviceSettings* DeviceSettings = GetDeviceSettings<ULiveLinkUnrealDeviceSettings>();
+	if (DeviceSettings->ClientId.IsValid())
+	{
+		ILiveLinkHubClientsModel& ClientsModel = ILiveLinkHubClientsModel::GetChecked();
+		return ClientsModel.GetClientDisplayName(DeviceSettings->ClientId).ToString();
+	}
+
+	return FString();
+}
+
+
+bool ULiveLinkUnrealDevice::SetHardwareId_Implementation(const FString& InHardwareID)
+{
+	return false;
+}
+
+
+bool ULiveLinkUnrealDevice::Connect_Implementation()
+{
+	if (ConnectionStatus == ELiveLinkDeviceConnectionStatus::Connected)
+	{
+		return false;
+	}
+
+	Connect();
+
+	return true;
+}
+
+
+bool ULiveLinkUnrealDevice::Disconnect_Implementation()
+{
+	if (ConnectionStatus == ELiveLinkDeviceConnectionStatus::Disconnected)
+	{
+		return false;
+	}
+
+	Disconnect();
+
+	return true;
+}
+
+void ULiveLinkUnrealDevice::SetConnectionStatus(ELiveLinkDeviceConnectionStatus InStatus)
+{
+	ConnectionStatus = InStatus;
+	ILiveLinkDeviceCapability_Connection::SetConnectionStatus(InStatus);
+}
+
+
+bool ULiveLinkUnrealDevice::StartRecording_Implementation()
+{
+	if (bIsRecording)
+	{
+		return false;
+	}
+
+	if (MessageEndpoint && DeviceAddress.IsValid())
+	{
+		FLiveLinkTakeRecorderCmd_StartRecording* Message = FMessageEndpoint::MakeMessage<FLiveLinkTakeRecorderCmd_StartRecording>();
+
+		ILiveLinkRecordingSession& Session = ILiveLinkRecordingSession::Get();
+		Message->SlateInfo = {
+			.SlateName = Session.GetSlateName(),
+			.TakeNumber = Session.GetTakeNumber(),
+			.Description = Session.GetSessionName(),
+		};
+		MessageEndpoint->Send(Message, EMessageFlags::Reliable, {}, nullptr, { DeviceAddress }, FTimespan::Zero(), FDateTime::MaxValue());
+		return true;
+	}
+
+	return false;
+}
+
+
+bool ULiveLinkUnrealDevice::StopRecording_Implementation()
+{
+	if (!bIsRecording)
+	{
+		return false;
+	}
+
+	if (MessageEndpoint && DeviceAddress.IsValid())
+	{
+		MessageEndpoint->Send(FMessageEndpoint::MakeMessage<FLiveLinkTakeRecorderCmd_StopRecording>(),
+			EMessageFlags::Reliable, {}, nullptr, { DeviceAddress }, FTimespan::Zero(), FDateTime::MaxValue());
+		return true;
+	}
+
+	return false;
+}
+
+
+bool ULiveLinkUnrealDevice::IsRecording_Implementation() const
+{
+	return bIsRecording;
+}
+
+
+void ULiveLinkUnrealDevice::HandleSlateNameChanged(FStringView InSlateName)
+{
+	if (MessageEndpoint && DeviceAddress.IsValid())
+	{
+		FLiveLinkTakeRecorderCmd_SetSlateName* Message = FMessageEndpoint::MakeMessage<FLiveLinkTakeRecorderCmd_SetSlateName>();
+		Message->SlateName = InSlateName;
+		MessageEndpoint->Send(Message, EMessageFlags::Reliable, {}, nullptr, { DeviceAddress }, FTimespan::Zero(), FDateTime::MaxValue());
+	}
+}
+
+
+void ULiveLinkUnrealDevice::HandleTakeNumberChanged(int32 InTakeNumber)
+{
+	if (MessageEndpoint && DeviceAddress.IsValid())
+	{
+		FLiveLinkTakeRecorderCmd_SetTakeNumber* Message = FMessageEndpoint::MakeMessage<FLiveLinkTakeRecorderCmd_SetTakeNumber>();
+		Message->TakeNumber = InTakeNumber;
+		MessageEndpoint->Send(Message, EMessageFlags::Reliable, {}, nullptr, { DeviceAddress }, FTimespan::Zero(), FDateTime::MaxValue());
+	}
+}
+
+
+void ULiveLinkUnrealDevice::HandleMessageCatchall(const TSharedRef<IMessageContext, ESPMode::ThreadSafe>& InContext)
+{
+	const UScriptStruct* MessageTypeInfo = InContext->GetMessageTypeInfo().Get();
+
+	if (MessageTypeInfo->IsChildOf(FLiveLinkHubAuxChannelBaseMessage::StaticStruct()))
+	{
+		const FLiveLinkHubAuxChannelBaseMessage* AuxBaseMessage =
+			static_cast<const FLiveLinkHubAuxChannelBaseMessage*>(InContext->GetMessage());
+
+		if (MessageTypeInfo->IsChildOf(FLiveLinkHubAuxChannelAcceptMessage::StaticStruct()))
+		{
+			ensure(AuxBaseMessage->ChannelId == ChannelId);
+
+			DeviceAddress = InContext->GetSender();
+
+			UE_LOGFMT(LogLiveLinkUnrealDevice, Log, "{DeviceName}: Connection accepted by {RemoteAddress}",
+				GetSettings()->DisplayName, DeviceAddress.ToString());
+
+			SetConnectionStatus(ELiveLinkDeviceConnectionStatus::Connected);
+		}
+		else if (MessageTypeInfo->IsChildOf(FLiveLinkHubAuxChannelRejectMessage::StaticStruct()))
+		{
+			ensure(AuxBaseMessage->ChannelId == ChannelId);
+
+			UE_LOGFMT(LogLiveLinkUnrealDevice, Error, "{DeviceName}: Connection rejected by {RemoteAddress}",
+				GetSettings()->DisplayName, InContext->GetSender().ToString());
+
+			OnDisconnect();
+		}
+		else if (MessageTypeInfo->IsChildOf(FLiveLinkHubAuxChannelCloseMessage::StaticStruct()))
+		{
+			ensure(AuxBaseMessage->ChannelId == ChannelId);
+
+			UE_LOGFMT(LogLiveLinkUnrealDevice, Log, "{DeviceName}: Connection closed by {RemoteAddress}",
+				GetSettings()->DisplayName, InContext->GetSender().ToString());
+
+			OnDisconnect();
+		}
+	}
+	else if (MessageTypeInfo->IsChildOf(FLiveLinkTakeRecorderMessageBase::StaticStruct()))
+	{
+		if (MessageTypeInfo == FLiveLinkTakeRecorderEvent_RecordingStarted::StaticStruct())
+		{
+			UE_LOGFMT(LogLiveLinkUnrealDevice, Log, "{DeviceName}: RecordingStarted event from {RemoteAddress}",
+				GetSettings()->DisplayName, InContext->GetSender().ToString());
+
+			bIsRecording = true;
+
+			if (GetDeviceSettings<ULiveLinkUnrealDeviceSettings>()->bHasRecordStartAuthority)
+			{
+				ILiveLinkRecordingSession& Session = ILiveLinkRecordingSession::Get();
+				const bool bWasRecording = Session.IsRecording();
+				const bool bCanRecord = Session.CanRecord();
+				bool bStartedRecording = false;
+				if (!bWasRecording && bCanRecord)
+				{
+					bStartedRecording = Session.StartRecording();
+				}
+
+				UE_LOGFMT(LogLiveLinkUnrealDevice, Log, "{DeviceName}: Has record start authority; "
+					"bWasRecording={bWasRecording}, bCanRecord={bCanRecord}, bStartedRecording={bStartedRecording}",
+					GetSettings()->DisplayName, bWasRecording, bCanRecord, bStartedRecording);
+			}
+		}
+		else if (MessageTypeInfo == FLiveLinkTakeRecorderEvent_RecordingStopped::StaticStruct())
+		{
+			UE_LOGFMT(LogLiveLinkUnrealDevice, Log, "{DeviceName}: RecordingStopped event from {RemoteAddress}",
+				GetSettings()->DisplayName, InContext->GetSender().ToString());
+
+			bIsRecording = false;
+
+			if (GetDeviceSettings<ULiveLinkUnrealDeviceSettings>()->bHasRecordStopAuthority)
+			{
+				ILiveLinkRecordingSession& Session = ILiveLinkRecordingSession::Get();
+				const bool bWasRecording = Session.IsRecording();
+				bool bStoppedRecording = false;
+				if (bWasRecording)
+				{
+					bStoppedRecording = ILiveLinkRecordingSession::Get().StopRecording();
+				}
+
+				UE_LOGFMT(LogLiveLinkUnrealDevice, Log, "{DeviceName}: Has record stop authority; "
+					"bWasRecording={bWasRecording}, bStoppedRecording={bStoppedRecording}",
+					GetSettings()->DisplayName, bWasRecording, bStoppedRecording);
+			}
+		}
+	}
+}
+
+
+void ULiveLinkUnrealDevice::HandleMessageBusNotification(const FMessageBusNotification& InNotification)
+{
+	if (InNotification.NotificationType == EMessageBusNotification::Unregistered)
+	{
+		if (InNotification.RegistrationAddress == DeviceAddress)
+		{
+			UE_LOGFMT(LogLiveLinkUnrealDevice, Log, "{DeviceName}: Message address unregistered {RemoteAddress}",
+				GetSettings()->DisplayName, DeviceAddress.ToString());
+
+			OnDisconnect();
+		}
+	}
+}
+
+
+void ULiveLinkUnrealDevice::Connect()
+{
+	if (ConnectionStatus != ELiveLinkDeviceConnectionStatus::Disconnected)
+	{
+		return;
+	}
+
+	const ULiveLinkUnrealDeviceSettings* DeviceSettings = GetDeviceSettings<ULiveLinkUnrealDeviceSettings>();
+	if (DeviceSettings->ClientId.IsValid())
+	{
+		ILiveLinkHubClientsModel& ClientsModel = ILiveLinkHubClientsModel::GetChecked();
+
+		UE_LOGFMT(LogLiveLinkUnrealDevice, Log, "{DeviceName}: Connecting to {ClientName}",
+			GetSettings()->DisplayName, ClientsModel.GetClientDisplayName(DeviceSettings->ClientId).ToString());
+
+		MessageEndpoint = FMessageEndpointBuilder(GetFName())
+			.WithCatchall(this, &ULiveLinkUnrealDevice::HandleMessageCatchall)
+			.NotificationHandling(FOnBusNotification::CreateUObject(this, &ULiveLinkUnrealDevice::HandleMessageBusNotification))
+			.ReceivingOnThread(ENamedThreads::GameThread)
+			.Build();
+
+		FLiveLinkUnrealDeviceAuxChannelRequestMessage* RequestMessage =
+			FMessageEndpoint::MakeMessage<FLiveLinkUnrealDeviceAuxChannelRequestMessage>();
+
+		ChannelId = FGuid::NewGuid();
+		RequestMessage->ChannelId = ChannelId;
+
+		ClientsModel.RequestAuxiliaryChannel(DeviceSettings->ClientId, *MessageEndpoint, *RequestMessage);
+
+		SetConnectionStatus(ELiveLinkDeviceConnectionStatus::Connecting);
+	}
+}
+
+
+void ULiveLinkUnrealDevice::Disconnect()
+{
+	if (MessageEndpoint && DeviceAddress.IsValid() && ChannelId.IsValid())
+	{
+		FLiveLinkHubAuxChannelCloseMessage* CloseMessage =
+			FMessageEndpoint::MakeMessage<FLiveLinkHubAuxChannelCloseMessage>();
+		CloseMessage->ChannelId = ChannelId;
+		MessageEndpoint->Send(CloseMessage, EMessageFlags::Reliable, {}, nullptr, { DeviceAddress }, FTimespan::Zero(), FDateTime::MaxValue());
+	}
+
+	OnDisconnect();
+}
+
+
+void ULiveLinkUnrealDevice::OnDisconnect()
+{
+	bIsRecording = false;
+
+	ChannelId.Invalidate();
+	DeviceAddress.Invalidate();
+	MessageEndpoint.Reset();
+
+	SetConnectionStatus(ELiveLinkDeviceConnectionStatus::Disconnected);
+}
+
+
+#undef LOCTEXT_NAMESPACE
diff --git a/Engine/Plugins/Animation/LiveLinkStandardDevices/Source/LiveLinkStandardDevices/Private/Devices/LiveLinkUnrealDevice.h b/Engine/Plugins/Animation/LiveLinkStandardDevices/Source/LiveLinkStandardDevices/Private/Devices/LiveLinkUnrealDevice.h
new file mode 100644
--- /dev/null
+++ b/Engine/Plugins/Animation/LiveLinkStandardDevices/Source/LiveLinkStandardDevices/Private/Devices/LiveLinkUnrealDevice.h
@@ -0,0 +1,128 @@
+// Copyright Epic Games, Inc. All Rights Reserved.
+
+#pragma once
+
+#include "LiveLinkDevice.h"
+#include "ILiveLinkHubClientsModel.h"
+#include "IMessageContext.h"
+#include "LiveLinkDeviceCapability_Connection.h"
+#include "LiveLinkDeviceCapability_Recording.h"
+#include "LiveLinkUnrealDevice.generated.h"
+
+
+DECLARE_LOG_CATEGORY_EXTERN(LogLiveLinkUnrealDevice, Log, All);
+
+
+struct FMessageBusNotification;
+class FMessageEndpoint;
+class IMessageContext;
+
+
+/** Settings for an Unreal device instance. */
+UCLASS()
+class ULiveLinkUnrealDeviceSettings : public ULiveLinkDeviceSettings
+{
+	GENERATED_BODY()
+
+public:
+	/** User specified display name for the device. */
+	UPROPERTY(EditAnywhere, Category="Unreal Device")
+	FString DisplayName = NSLOCTEXT("LiveLinkUnrealDevice", "DefaultDisplayName", "Unreal Editor").ToString();
+
+	/** The editor client known to Live Link Hub to connect to. */
+	UPROPERTY(EditAnywhere, Category="Unreal Device")
+	FLiveLinkHubClientId ClientId;
+
+	/** If true, a "recording started" event from this device will also cause Live Link Hub to begin recording. */
+	UPROPERTY(EditAnywhere, Category="Unreal Device", AdvancedDisplay)
+	bool bHasRecordStartAuthority = true;
+
+	/** If true, a "recording stopped" event from this device will also cause Live Link Hub to stop recording. */
+	UPROPERTY(EditAnywhere, Category="Unreal Device", AdvancedDisplay)
+	bool bHasRecordStopAuthority = true;
+};
+
+
+/**
+ * A device for connecting to a running Unreal Editor instance.
+ * 
+ * Supports sending/receiving Take Recorder events to/from the remote editor session.
+ */
+UCLASS(meta=(DisplayName="Unreal"))
+class ULiveLinkUnrealDevice : public ULiveLinkDevice
+	, public ILiveLinkDeviceCapability_Connection
+	, public ILiveLinkDeviceCapability_Recording
+{
+	GENERATED_BODY()
+
+public:
+	ULiveLinkUnrealDeviceSettings* GetSettings() { return GetDeviceSettings<ULiveLinkUnrealDeviceSettings>(); }
+	const ULiveLinkUnrealDeviceSettings* GetSettings() const { return GetDeviceSettings<ULiveLinkUnrealDeviceSettings>(); }
+
+	//~ Begin ULiveLinkDevice interface
+public:
+	virtual TSubclassOf<ULiveLinkDeviceSettings> GetSettingsClass() const override { return ULiveLinkUnrealDeviceSettings::StaticClass(); }
+
+	virtual FText GetDisplayName() const override;
+	virtual EDeviceHealth GetDeviceHealth() const override;
+	virtual FText GetHealthText() const override;
+
+	virtual void OnDeviceAdded() override;
+	virtual void OnDeviceRemoved() override;
+	virtual void OnSettingChanged(const FPropertyChangedEvent& InPropertyChangedEvent) override;
+	//~ End ULiveLinkDevice interface
+
+	//~ Begin ILiveLinkDeviceCapability_Connection interface
+public:
+	virtual ELiveLinkDeviceConnectionStatus GetConnectionStatus_Implementation() const override;
+	virtual FString GetHardwareId_Implementation() const override;
+	virtual bool SetHardwareId_Implementation(const FString& InHardwareID) override;
+	virtual bool Connect_Implementation() override;
+	virtual bool Disconnect_Implementation() override;
+protected:
+	virtual void SetConnectionStatus(ELiveLinkDeviceConnectionStatus InStatus) override;
+	//~ End ILiveLinkDeviceCapability_Connection interface
+
+	//~ Begin ILiveLinkDeviceCapability_Recording interface
+public:
+	virtual bool StartRecording_Implementation() override;
+	virtual bool StopRecording_Implementation() override;
+	virtual bool IsRecording_Implementation() const override;
+	//~ End ILiveLinkDeviceCapability_Recording interface
+
+private:
+	TPair<EDeviceHealth, FText> GetHealth() const;
+
+	// Take Recorder event handlers.
+	void HandleSlateNameChanged(FStringView InSlateName);
+	void HandleTakeNumberChanged(int32 InTakeNumber);
+
+	// Message bus event handlers.
+	void HandleMessageCatchall(const TSharedRef<IMessageContext, ESPMode::ThreadSafe>& InContext);
+	void HandleMessageBusNotification(const FMessageBusNotification& InNotification);
+
+	/** Attempts to connect to the the configured `ClientId`. */
+	void Connect();
+
+	/** Sends a close message connected editor client (if any). */
+	void Disconnect();
+
+	/** Called to perform cleanup and state transitions for both expected and unexpected disconnections. */
+	void OnDisconnect();
+
+private:
+	/** Updated in response to internal state transitions, and reflected in the device UI. */
+	ELiveLinkDeviceConnectionStatus ConnectionStatus = ELiveLinkDeviceConnectionStatus::Disconnected;
+
+	/** The endpoint used for communicating with the editor. */
+	TSharedPtr<FMessageEndpoint> MessageEndpoint;
+
+	/** The address that accepted the aux channel request; we Send() here rather than Publish(). */
+	FMessageAddress DeviceAddress;
+
+	/** The aux channel identifier. */
+	FGuid ChannelId;
+
+	/** Recording state is updated in response to device event messages. */
+	bool bIsRecording = false;
+};
diff --git a/Engine/Plugins/Animation/LiveLinkStandardDevices/Source/LiveLinkStandardDevices/Private/LiveLinkHubUnrealDeviceAux.cpp b/Engine/Plugins/Animation/LiveLinkStandardDevices/Source/LiveLinkStandardDevices/Private/LiveLinkHubUnrealDeviceAux.cpp
new file mode 100644
--- /dev/null
+++ b/Engine/Plugins/Animation/LiveLinkStandardDevices/Source/LiveLinkStandardDevices/Private/LiveLinkHubUnrealDeviceAux.cpp
@@ -0,0 +1,260 @@
+// Copyright Epic Games, Inc. All Rights Reserved.
+
+#include "LiveLinkHubUnrealDeviceAux.h"
+#include "Devices/LiveLinkUnrealDevice.h"
+#include "Engine/Engine.h"
+#include "ILiveLinkHubMessagingModule.h"
+#include "LiveLinkUnrealDeviceMessages.h"
+#include "Logging/StructuredLog.h"
+#include "MessageEndpointBuilder.h"
+#include "Recorder/TakeRecorderSubsystem.h"
+
+
+FLiveLinkHubUnrealDeviceAuxManager::FLiveLinkHubUnrealDeviceAuxManager()
+{
+	FMessageEndpointBuilder EndpointBuilder = FMessageEndpointBuilder("LiveLinkHubTakeRecorderAuxHandler");
+	MessageEndpoint = EndpointBuilder
+		.Handling<FLiveLinkTakeRecorderCmd_SetSlateName>(this, &FLiveLinkHubUnrealDeviceAuxManager::HandleSetSlateName)
+		.Handling<FLiveLinkTakeRecorderCmd_SetTakeNumber>(this, &FLiveLinkHubUnrealDeviceAuxManager::HandleSetTakeNumber)
+		.Handling<FLiveLinkTakeRecorderCmd_StartRecording>(this, &FLiveLinkHubUnrealDeviceAuxManager::HandleStartRecording)
+		.Handling<FLiveLinkTakeRecorderCmd_StopRecording>(this, &FLiveLinkHubUnrealDeviceAuxManager::HandleStopRecording)
+		.Handling<FLiveLinkHubAuxChannelCloseMessage>(this, &FLiveLinkHubUnrealDeviceAuxManager::HandleAuxClose)
+		.ReceivingOnThread(ENamedThreads::GameThread)
+		.Build();
+
+	RegisterRequestHandler();
+	RegisterTakeRecorderDelegates();
+}
+
+
+FLiveLinkHubUnrealDeviceAuxManager::~FLiveLinkHubUnrealDeviceAuxManager()
+{
+	ILiveLinkHubMessagingModule* HubMessagingModule = FModuleManager::Get().GetModulePtr<ILiveLinkHubMessagingModule>("LiveLinkHubMessaging");
+	if (HubMessagingModule)
+	{
+		HubMessagingModule->UnregisterAuxChannelRequestHandler<FLiveLinkUnrealDeviceAuxChannelRequestMessage>();
+	}
+
+	if (GEngine)
+	{
+		if (UTakeRecorderSubsystem* TakeRecorder = GEngine->GetEngineSubsystem<UTakeRecorderSubsystem>())
+		{
+			TakeRecorder->GetOnRecordingStartedEvent().RemoveAll(this);
+			TakeRecorder->GetOnRecordingStoppedEvent().RemoveAll(this);
+		}
+	}
+}
+
+
+void FLiveLinkHubUnrealDeviceAuxManager::RegisterRequestHandler()
+{
+	ILiveLinkHubMessagingModule& HubMessagingModule =
+		FModuleManager::Get().GetModuleChecked<ILiveLinkHubMessagingModule>("LiveLinkHubMessaging");
+
+	HubMessagingModule.RegisterAuxChannelRequestHandler<FLiveLinkUnrealDeviceAuxChannelRequestMessage>(
+		[this]
+		(
+			const FLiveLinkUnrealDeviceAuxChannelRequestMessage& InRequest,
+			const TSharedRef<IMessageContext, ESPMode::ThreadSafe>& InContext
+		)
+		{
+			UE_LOGFMT(LogLiveLinkUnrealDevice, Log, "Received Take Recorder channel request from {Sender}",
+				InContext->GetSender().ToString());
+
+			const FGuid& ChannelId = InRequest.ChannelId;
+
+			if (FMessageAddress* ExistingAddress = ChannelToAddress.Find(ChannelId))
+			{
+				// Shouldn't happen.
+				UE_LOGFMT(LogLiveLinkUnrealDevice, Warning, "Duplicate channel ID {ChannelId} ({ExistingAddress})",
+					ChannelId.ToString(), ExistingAddress->ToString());
+
+				AddressToChannel.Remove(*ExistingAddress);
+				ChannelToAddress.Remove(ChannelId);
+			}
+
+			if (FGuid* ExistingChannel = AddressToChannel.Find(InContext->GetSender()))
+			{
+				// Shouldn't happen.
+				UE_LOGFMT(LogLiveLinkUnrealDevice, Warning, "Duplicate aux address {SenderAddress} ({ExistingChannel})",
+					InContext->GetSender().ToString(), ExistingChannel->ToString());
+
+				ChannelToAddress.Remove(*ExistingChannel);
+				AddressToChannel.Remove(InContext->GetSender());
+			}
+
+			ChannelToAddress.Add(ChannelId, InContext->GetSender());
+			AddressToChannel.Add(InContext->GetSender(), ChannelId);
+
+			FLiveLinkHubAuxChannelAcceptMessage* AcceptMessage =
+				FMessageEndpoint::MakeMessage<FLiveLinkHubAuxChannelAcceptMessage>();
+			AcceptMessage->ChannelId = ChannelId;
+			MessageEndpoint->Send(AcceptMessage, EMessageFlags::Reliable, {}, nullptr, { InContext->GetSender() },
+				FTimespan::Zero(), FDateTime::MaxValue());
+		}
+	);
+}
+
+
+void FLiveLinkHubUnrealDeviceAuxManager::RegisterTakeRecorderDelegates()
+{
+	if (UTakeRecorderSubsystem* TakeRecorder = GEngine->GetEngineSubsystem<UTakeRecorderSubsystem>(); ensure(TakeRecorder))
+	{
+		TakeRecorder->GetOnRecordingStartedEvent().AddRaw(this, &FLiveLinkHubUnrealDeviceAuxManager::OnRecordingStarted);
+		TakeRecorder->GetOnRecordingStoppedEvent().AddRaw(this, &FLiveLinkHubUnrealDeviceAuxManager::OnRecordingStopped);
+	}
+}
+
+
+bool FLiveLinkHubUnrealDeviceAuxManager::IsKnownSender(const FMessageAddress& InAddress) const
+{
+	if (AddressToChannel.Contains(InAddress))
+	{
+		return true;
+	}
+
+	UE_LOGFMT(LogLiveLinkUnrealDevice, Warning, "Message from unknown sender {Sender} will be ignored", InAddress.ToString());
+	return false;
+}
+
+
+void FLiveLinkHubUnrealDeviceAuxManager::HandleAuxClose(
+	const FLiveLinkHubAuxChannelCloseMessage& InMessage,
+	const TSharedRef<IMessageContext, ESPMode::ThreadSafe>& InContext
+)
+{
+	UE_LOGFMT(LogLiveLinkUnrealDevice, Log, "Received channel close from {Sender}", InContext->GetSender().ToString());
+
+	if (!IsKnownSender(InContext->GetSender()))
+	{
+		return;
+	}
+
+	if (FGuid* ExistingChannel = AddressToChannel.Find(InContext->GetSender()))
+	{
+		if (*ExistingChannel != InMessage.ChannelId)
+		{
+			UE_LOGFMT(LogLiveLinkUnrealDevice, Warning, "Channel close has wrong ID {ChannelId} (expected {ExistingId})",
+				InMessage.ChannelId.ToString(), ExistingChannel->ToString());
+		}
+
+		AddressToChannel.Remove(InContext->GetSender());
+	}
+	else
+	{
+		// Shouldn't happen.
+		UE_LOGFMT(LogLiveLinkUnrealDevice, Warning, "Channel with address {Sender} not found", InContext->GetSender().ToString());
+	}
+
+	if (FMessageAddress* ExistingAddress = ChannelToAddress.Find(InMessage.ChannelId))
+	{
+		if (*ExistingAddress != InContext->GetSender())
+		{
+			UE_LOGFMT(LogLiveLinkUnrealDevice, Warning, "Channel close not from expected sender {ExistingAddress}",
+				ExistingAddress->ToString());
+		}
+
+		ChannelToAddress.Remove(InMessage.ChannelId);
+	}
+	else
+	{
+		// Shouldn't happen.
+		UE_LOGFMT(LogLiveLinkUnrealDevice, Warning, "Channel with ID {ChannelId} not found", InMessage.ChannelId.ToString());
+	}
+}
+
+
+void FLiveLinkHubUnrealDeviceAuxManager::HandleSetSlateName(
+	const FLiveLinkTakeRecorderCmd_SetSlateName& InCmd,
+	const TSharedRef<IMessageContext, ESPMode::ThreadSafe>& InContext
+)
+{
+	UE_LOGFMT(LogLiveLinkUnrealDevice, Log, "Received SetSlateName from {Sender}", InContext->GetSender().ToString());
+
+	if (!IsKnownSender(InContext->GetSender()))
+	{
+		return;
+	}
+
+	if (UTakeRecorderSubsystem* TakeRecorder = GEngine->GetEngineSubsystem<UTakeRecorderSubsystem>(); ensure(TakeRecorder))
+	{
+		TakeRecorder->SetSlateName(InCmd.SlateName);
+	}
+}
+
+
+void FLiveLinkHubUnrealDeviceAuxManager::HandleSetTakeNumber(
+	const FLiveLinkTakeRecorderCmd_SetTakeNumber& InCmd,
+	const TSharedRef<IMessageContext, ESPMode::ThreadSafe>& InContext
+)
+{
+	UE_LOGFMT(LogLiveLinkUnrealDevice, Log, "Received SetTakeNumber from {Sender}", InContext->GetSender().ToString());
+
+	if (!IsKnownSender(InContext->GetSender()))
+	{
+		return;
+	}
+
+	if (UTakeRecorderSubsystem* TakeRecorder = GEngine->GetEngineSubsystem<UTakeRecorderSubsystem>(); ensure(TakeRecorder))
+	{
+		TakeRecorder->SetTakeNumber(InCmd.TakeNumber);
+	}
+}
+
+
+void FLiveLinkHubUnrealDeviceAuxManager::HandleStartRecording(
+	const FLiveLinkTakeRecorderCmd_StartRecording& InCmd,
+	const TSharedRef<IMessageContext, ESPMode::ThreadSafe>& InContext
+)
+{
+	UE_LOGFMT(LogLiveLinkUnrealDevice, Log, "Received StartRecording from {Sender}", InContext->GetSender().ToString());
+
+	if (!IsKnownSender(InContext->GetSender()))
+	{
+		return;
+	}
+
+	if (UTakeRecorderSubsystem* TakeRecorder = GEngine->GetEngineSubsystem<UTakeRecorderSubsystem>(); ensure(TakeRecorder))
+	{
+		TakeRecorder->StartRecording();
+	}
+}
+
+
+void FLiveLinkHubUnrealDeviceAuxManager::HandleStopRecording(
+	const FLiveLinkTakeRecorderCmd_StopRecording& InCmd,
+	const TSharedRef<IMessageContext, ESPMode::ThreadSafe>& InContext
+)
+{
+	UE_LOGFMT(LogLiveLinkUnrealDevice, Log, "Received StopRecording from {Sender}", InContext->GetSender().ToString());
+
+	if (!IsKnownSender(InContext->GetSender()))
+	{
+		return;
+	}
+
+	if (UTakeRecorderSubsystem* TakeRecorder = GEngine->GetEngineSubsystem<UTakeRecorderSubsystem>(); ensure(TakeRecorder))
+	{
+		TakeRecorder->StopRecording();
+	}
+}
+
+
+void FLiveLinkHubUnrealDeviceAuxManager::OnRecordingStarted(UTakeRecorder* InRecorder)
+{
+	TArray<FMessageAddress> ChannelAddresses;
+	ChannelToAddress.GenerateValueArray(ChannelAddresses);
+
+	FLiveLinkTakeRecorderEvent_RecordingStarted* EventMessage = FMessageEndpoint::MakeMessage<FLiveLinkTakeRecorderEvent_RecordingStarted>();
+	MessageEndpoint->Send(EventMessage, EMessageFlags::Reliable, {}, nullptr, ChannelAddresses, FTimespan::Zero(), FDateTime::MaxValue());
+}
+
+
+void FLiveLinkHubUnrealDeviceAuxManager::OnRecordingStopped(UTakeRecorder* InRecorder)
+{
+	TArray<FMessageAddress> ChannelAddresses;
+	ChannelToAddress.GenerateValueArray(ChannelAddresses);
+
+	FLiveLinkTakeRecorderEvent_RecordingStopped* EventMessage = FMessageEndpoint::MakeMessage<FLiveLinkTakeRecorderEvent_RecordingStopped>();
+	MessageEndpoint->Send(EventMessage, EMessageFlags::Reliable, {}, nullptr, ChannelAddresses, FTimespan::Zero(), FDateTime::MaxValue());
+}
diff --git a/Engine/Plugins/Animation/LiveLinkStandardDevices/Source/LiveLinkStandardDevices/Private/LiveLinkHubUnrealDeviceAux.h b/Engine/Plugins/Animation/LiveLinkStandardDevices/Source/LiveLinkStandardDevices/Private/LiveLinkHubUnrealDeviceAux.h
new file mode 100644
--- /dev/null
+++ b/Engine/Plugins/Animation/LiveLinkStandardDevices/Source/LiveLinkStandardDevices/Private/LiveLinkHubUnrealDeviceAux.h
@@ -0,0 +1,55 @@
+// Copyright Epic Games, Inc. All Rights Reserved.
+
+#pragma once
+
+#include "IMessageContext.h"
+#include "Misc/Guid.h"
+#include "Templates/SharedPointer.h"
+
+
+struct FLiveLinkHubAuxChannelCloseMessage;
+struct FLiveLinkTakeRecorderCmd_SetSlateName;
+struct FLiveLinkTakeRecorderCmd_SetTakeNumber;
+struct FLiveLinkTakeRecorderCmd_StartRecording;
+struct FLiveLinkTakeRecorderCmd_StopRecording;
+class FMessageEndpoint;
+class IMessageContext;
+class UTakeRecorder;
+
+
+/**
+ * Configures a Live Link Hub auxiliary message channel for communicating with Unreal devices.
+ * 
+ * Handles incoming Take Recorder commands, and relays events to connected devices.
+ */
+class FLiveLinkHubUnrealDeviceAuxManager
+{
+public:
+	FLiveLinkHubUnrealDeviceAuxManager();
+	~FLiveLinkHubUnrealDeviceAuxManager();
+
+private:
+	// Init helpers
+	void RegisterRequestHandler();
+	void RegisterTakeRecorderDelegates();
+
+	// Take Recorder event handlers
+	void OnRecordingStarted(UTakeRecorder* InRecorder);
+	void OnRecordingStopped(UTakeRecorder* InRecorder);
+
+	// Message handlers
+	bool IsKnownSender(const FMessageAddress& InAddress) const;
+
+	void HandleAuxClose(const FLiveLinkHubAuxChannelCloseMessage& InMessage, const TSharedRef<IMessageContext, ESPMode::ThreadSafe>& InContext);
+
+	void HandleSetSlateName(const FLiveLinkTakeRecorderCmd_SetSlateName& InCmd, const TSharedRef<IMessageContext, ESPMode::ThreadSafe>& InContext);
+	void HandleSetTakeNumber(const FLiveLinkTakeRecorderCmd_SetTakeNumber& InCmd, const TSharedRef<IMessageContext, ESPMode::ThreadSafe>& InContext);
+	void HandleStartRecording(const FLiveLinkTakeRecorderCmd_StartRecording& InCmd, const TSharedRef<IMessageContext, ESPMode::ThreadSafe>& InContext);
+	void HandleStopRecording(const FLiveLinkTakeRecorderCmd_StopRecording& InCmd, const TSharedRef<IMessageContext, ESPMode::ThreadSafe>& InContext);
+
+private:
+	TSharedPtr<FMessageEndpoint> MessageEndpoint;
+
+	TMap<FGuid, FMessageAddress> ChannelToAddress;
+	TMap<FMessageAddress, FGuid> AddressToChannel;
+};
diff --git a/Engine/Plugins/Animation/LiveLinkStandardDevices/Source/LiveLinkStandardDevices/Private/LiveLinkStandardDevicesModule.cpp b/Engine/Plugins/Animation/LiveLinkStandardDevices/Source/LiveLinkStandardDevices/Private/LiveLinkStandardDevicesModule.cpp
new file mode 100644
--- /dev/null
+++ b/Engine/Plugins/Animation/LiveLinkStandardDevices/Source/LiveLinkStandardDevices/Private/LiveLinkStandardDevicesModule.cpp
@@ -0,0 +1,36 @@
+// Copyright Epic Games, Inc. All Rights Reserved.
+
+#include "LiveLinkHubUnrealDeviceAux.h"
+#include "Misc/ConfigCacheIni.h"
+#include "Misc/CoreDelegates.h"
+#include "Modules/ModuleManager.h"
+
+
+class FLiveLinkStandardDevicesModule : public IModuleInterface
+{
+public:
+	//~ Begin IModuleInterface
+	virtual ~FLiveLinkStandardDevicesModule() = default;
+
+	virtual void StartupModule()
+	{
+		FCoreDelegates::OnPostEngineInit.AddLambda(
+			[this]()
+			{
+				UnrealDeviceAuxManager = MakeUnique<FLiveLinkHubUnrealDeviceAuxManager>();
+			}
+		);
+	}
+
+	virtual void PreUnloadCallback()
+	{
+		UnrealDeviceAuxManager.Reset();
+	}
+	//~ End IModuleInterface
+
+private:
+	TUniquePtr<FLiveLinkHubUnrealDeviceAuxManager> UnrealDeviceAuxManager;
+};
+
+
+IMPLEMENT_MODULE(FLiveLinkStandardDevicesModule, LiveLinkStandardDevices);
diff --git a/Engine/Plugins/Animation/LiveLinkStandardDevices/Source/LiveLinkStandardDevices/Private/LiveLinkUnrealDeviceMessages.h b/Engine/Plugins/Animation/LiveLinkStandardDevices/Source/LiveLinkStandardDevices/Private/LiveLinkUnrealDeviceMessages.h
new file mode 100644
--- /dev/null
+++ b/Engine/Plugins/Animation/LiveLinkStandardDevices/Source/LiveLinkStandardDevices/Private/LiveLinkUnrealDeviceMessages.h
@@ -0,0 +1,122 @@
+// Copyright Epic Games, Inc. All Rights Reserved.
+
+#pragma once
+
+#include "CoreTypes.h"
+#include "LiveLinkHubMessages.h"
+#include "Misc/Timecode.h"
+#include "LiveLinkUnrealDeviceMessages.generated.h"
+
+
+USTRUCT()
+struct FLiveLinkUnrealDeviceAuxChannelRequestMessage : public FLiveLinkHubAuxChannelRequestMessage
+{
+	GENERATED_BODY()
+};
+
+
+USTRUCT()
+struct FLiveLinkTakeRecorderMessageBase
+{
+	GENERATED_BODY()
+
+	static constexpr uint8 CurrentVersion = 1;
+
+	UPROPERTY()
+	uint8 MessageVersion = CurrentVersion;
+
+	UPROPERTY()
+	FGuid MessageId;
+};
+
+
+USTRUCT()
+struct FLiveLinkTakeRecorderSlateInfo
+{
+	GENERATED_BODY()
+
+	UPROPERTY()
+	TOptional<FString> SlateName;
+
+	UPROPERTY()
+	TOptional<int32> TakeNumber;
+
+	UPROPERTY()
+	TOptional<FString> Description;
+};
+
+
+//////////////////////////////////////////////////////////////////////////
+
+
+USTRUCT()
+struct FLiveLinkTakeRecorderCmd_SetSlateName : public FLiveLinkTakeRecorderMessageBase
+{
+	GENERATED_BODY()
+
+	UPROPERTY()
+	FString SlateName;
+};
+
+
+USTRUCT()
+struct FLiveLinkTakeRecorderCmd_SetTakeNumber : public FLiveLinkTakeRecorderMessageBase
+{
+	GENERATED_BODY()
+
+	UPROPERTY()
+	int32 TakeNumber = 0;
+};
+
+
+USTRUCT()
+struct FLiveLinkTakeRecorderCmd_StartRecording : public FLiveLinkTakeRecorderMessageBase
+{
+	GENERATED_BODY()
+
+	UPROPERTY()
+	FLiveLinkTakeRecorderSlateInfo SlateInfo;
+};
+
+
+USTRUCT()
+struct FLiveLinkTakeRecorderCmd_StopRecording : public FLiveLinkTakeRecorderMessageBase
+{
+	GENERATED_BODY()
+};
+
+
+//////////////////////////////////////////////////////////////////////////
+
+
+USTRUCT()
+struct FLiveLinkTakeRecorderEvent_RecordingStarting : public FLiveLinkTakeRecorderMessageBase
+{
+	GENERATED_BODY()
+
+	UPROPERTY()
+	float CountdownSec = 0.0;
+};
+
+
+USTRUCT()
+struct FLiveLinkTakeRecorderEvent_RecordingStarted : public FLiveLinkTakeRecorderMessageBase
+{
+	GENERATED_BODY()
+
+	UPROPERTY()
+	FTimecode StartTimecode;
+};
+
+
+USTRUCT()
+struct FLiveLinkTakeRecorderEvent_RecordingStopped : public FLiveLinkTakeRecorderMessageBase
+{
+	GENERATED_BODY()
+
+	UPROPERTY()
+	FTimecode EndTimecode;
+
+	UPROPERTY()
+	bool bCancelled = false;
+};
diff --git a/Engine/Plugins/Animation/LiveLinkStandardDevices/Source/LiveLinkStandardDevices/Private/Tests/LiveLinkUnrealDeviceTests.cpp b/Engine/Plugins/Animation/LiveLinkStandardDevices/Source/LiveLinkStandardDevices/Private/Tests/LiveLinkUnrealDeviceTests.cpp
new file mode 100644
--- /dev/null
+++ b/Engine/Plugins/Animation/LiveLinkStandardDevices/Source/LiveLinkStandardDevices/Private/Tests/LiveLinkUnrealDeviceTests.cpp
@@ -0,0 +1,16 @@
+// Copyright Epic Games, Inc. All Rights Reserved.
+
+#include "Devices/LiveLinkUnrealDevice.h"
+#include "Misc/AutomationTest.h"
+
+
+IMPLEMENT_SIMPLE_AUTOMATION_TEST(FUnrealDeviceTest, "LiveLinkHub.UnrealDevice",
+	EAutomationTestFlags::EditorContext | EAutomationTestFlags::EngineFilter)
+
+
+bool FUnrealDeviceTest::RunTest(const FString& Parameters)
+{
+	ULiveLinkUnrealDevice* Device = NewObject<ULiveLinkUnrealDevice>();
+
+	return true;
+}
diff --git a/Engine/Source/Programs/LiveLinkHub/Source/LiveLinkHub.Target.cs b/Engine/Source/Programs/LiveLinkHub/Source/LiveLinkHub.Target.cs
--- a/Engine/Source/Programs/LiveLinkHub/Source/LiveLinkHub.Target.cs
+++ b/Engine/Source/Programs/LiveLinkHub/Source/LiveLinkHub.Target.cs
@@ -53,11 +53,12 @@ public class LiveLinkHubTarget : TargetRules
 			"XInputDevice",
 
 			"LiveLinkDevice",
+			"LiveLinkStandardDevices",
 
 			"LiveLinkOpenVR",
 
 			"LiveLinkOpenTrackIO",
-			
+
 			"CaptureManagerApp",
 			"CaptureManagerDevices",
 			"MediaIOFramework",
